fix non-unity build.

diff --git a/Engine/Source/Runtime/PhysXCooking/Private/PhysXCooking.cpp b/Engine/Source/Runtime/PhysXCooking/Private/PhysXCooking.cpp
--- a/Engine/Source/Runtime/PhysXCooking/Private/PhysXCooking.cpp
+++ b/Engine/Source/Runtime/PhysXCooking/Private/PhysXCooking.cpp
@@ -1,6 +1,7 @@
 // Copyright Epic Games, Inc. All Rights Reserved.
 
 #include "PhysXCooking.h"
+#include "Core/Public/Modules/ModuleManager.h"
 
 #if PHYSICS_INTERFACE_PHYSX
 
